Fix copyright headers

#ue5
#android

diff --git a/Engine/Plugins/Runtime/AndroidFileServer/Source/AndroidFileServer/AndroidFileServer.Build.cs b/Engine/Plugins/Runtime/AndroidFileServer/Source/AndroidFileServer/AndroidFileServer.Build.cs
--- a/Engine/Plugins/Runtime/AndroidFileServer/Source/AndroidFileServer/AndroidFileServer.Build.cs
+++ b/Engine/Plugins/Runtime/AndroidFileServer/Source/AndroidFileServer/AndroidFileServer.Build.cs
@@ -1,4 +1,4 @@
-// Some copyright should be here...
+// Copyright Epic Games, Inc. All Rights Reserved.
 
 using UnrealBuildTool;
 using System.IO;
