Fix up uplugin allow list for Android File Server

#ue5
#android

diff --git a/Engine/Plugins/Runtime/AndroidFileServer/AndroidFileServer.uplugin b/Engine/Plugins/Runtime/AndroidFileServer/AndroidFileServer.uplugin
--- a/Engine/Plugins/Runtime/AndroidFileServer/AndroidFileServer.uplugin
+++ b/Engine/Plugins/Runtime/AndroidFileServer/AndroidFileServer.uplugin
@@ -19,13 +19,13 @@
 			"Name": "AndroidFileServer",
 			"Type": "Runtime",
 			"LoadingPhase": "PreLoadingScreen",
-                        "WhitelistPlatforms" : [ "Android", "Win64", "Mac", "Linux" ]
+                        "PlatformAllowList" : [ "Android", "Win64", "Mac", "Linux" ]
                 },
                 {
                         "Name" : "AndroidFileServerEditor",
                         "Type" : "Editor",
                         "LoadingPhase" : "PostEngineInit",
-                        "WhitelistPlatforms" : [ "Win64", "Mac", "Linux" ]
+                        "PlatformAllowList" : [ "Win64", "Mac", "Linux" ]
                 }
         ]
 }
