Fix read-only issues with DefaultEngine.ini from AFS

#android

diff --git a/Engine/Plugins/Runtime/AndroidFileServer/Source/AndroidFileServerEditor/Private/AndroidFileServerRuntimeSettings.cpp b/Engine/Plugins/Runtime/AndroidFileServer/Source/AndroidFileServerEditor/Private/AndroidFileServerRuntimeSettings.cpp
--- a/Engine/Plugins/Runtime/AndroidFileServer/Source/AndroidFileServerEditor/Private/AndroidFileServerRuntimeSettings.cpp
+++ b/Engine/Plugins/Runtime/AndroidFileServer/Source/AndroidFileServerEditor/Private/AndroidFileServerRuntimeSettings.cpp
@@ -31,6 +31,6 @@ void UAndroidFileServerRuntimeSettings::PostInitProperties()
 	{
 		FGuid UniqueId = FGuid::NewGuid();
 		SecurityToken = UniqueId.ToString();
-		TryUpdateDefaultConfigFile();
+		TryUpdateDefaultConfigFile(TEXT(""), false);
 	}
 }
\ No newline at end of file
