Fix compiling of AFSProject on Linux

#android
#linux

diff --git a/Engine/Plugins/Runtime/AndroidFileServer/Source/AndroidFileServer/AndroidFileServer_UPL.xml b/Engine/Plugins/Runtime/AndroidFileServer/Source/AndroidFileServer/AndroidFileServer_UPL.xml
--- a/Engine/Plugins/Runtime/AndroidFileServer/Source/AndroidFileServer/AndroidFileServer_UPL.xml
+++ b/Engine/Plugins/Runtime/AndroidFileServer/Source/AndroidFileServer/AndroidFileServer_UPL.xml
@@ -176,7 +176,7 @@ tasks.whenTaskAdded { task ->
 			if (System.getProperty('os.name').toLowerCase(Locale.ROOT).contains('windows')) {
 				commandLine 'cmd', '/c', 'gradlew.bat', 'app:assembleRelease'
 			} else {
-				commandLine 'sh', '-c', 'gradlew', 'app:assembleRelease'
+				commandLine 'sh', '-c', 'chmod +x ./gradlew && ./gradlew', 'app:assembleRelease'
 			}
 		}
 		task.finalizedBy 'ueAFSProjectAssembleRelease'
@@ -188,7 +188,7 @@ tasks.whenTaskAdded { task ->
 			if (System.getProperty('os.name').toLowerCase(Locale.ROOT).contains('windows')) {
 				commandLine 'cmd', '/c', 'gradlew.bat', 'app:assembleDebug'
 			} else {
-				commandLine 'sh', '-c', 'gradlew', 'app:assembleDebug'
+				commandLine 'sh', '-c', 'chmod +x ./gradlew && ./gradlew', 'app:assembleDebug'
 			}
 		}
 		task.finalizedBy 'ueAFSProjectAssembleDebug'
@@ -200,7 +200,7 @@ tasks.whenTaskAdded { task ->
 			if (System.getProperty('os.name').toLowerCase(Locale.ROOT).contains('windows')) {
 				commandLine 'cmd', '/c', 'gradlew.bat', 'app:assembleRelease'
 			} else {
-				commandLine 'sh', '-c', 'gradlew', 'app:assembleRelease'
+				commandLine 'sh', '-c', 'chmod +x ./gradlew && ./gradlew', 'app:assembleRelease'
 			}
 		}
 		task.finalizedBy 'ueAFSProjectBundleRelease'
@@ -212,7 +212,7 @@ tasks.whenTaskAdded { task ->
 			if (System.getProperty('os.name').toLowerCase(Locale.ROOT).contains('windows')) {
 				commandLine 'cmd', '/c', 'gradlew.bat', 'app:assembleDebug'
 			} else {
-				commandLine 'sh', '-c', 'gradlew', 'app:assembleDebug'
+				commandLine 'sh', '-c', 'chmod +x ./gradlew && ./gradlew', 'app:assembleDebug'
 			}
 		}
 		task.finalizedBy 'ueAFSProjectBundleDebug'
