Fix EditCondition error in AndroidFileServer project settings
Change "Enable Plugin" to "Use AndroidFileServer" to reduce confusion with Plugin Settings

#android

diff --git a/Engine/Plugins/Runtime/AndroidFileServer/Source/AndroidFileServerEditor/Public/AndroidFileServerRuntimeSettings.h b/Engine/Plugins/Runtime/AndroidFileServer/Source/AndroidFileServerEditor/Public/AndroidFileServerRuntimeSettings.h
--- a/Engine/Plugins/Runtime/AndroidFileServer/Source/AndroidFileServerEditor/Public/AndroidFileServerRuntimeSettings.h
+++ b/Engine/Plugins/Runtime/AndroidFileServer/Source/AndroidFileServerEditor/Public/AndroidFileServerRuntimeSettings.h
@@ -29,8 +29,8 @@ class UAndroidFileServerRuntimeSettings : public UObject
 {
 	GENERATED_UCLASS_BODY()
 
-	// Enable Android FileServer plugin
-	UPROPERTY(EditAnywhere, config, Category = Packaging)
+	// Enable Android FileServer for packaged builds and quick launch
+	UPROPERTY(EditAnywhere, config, Category = Packaging, Meta = (DisplayName = "Use AndroidFileServer"))
 	bool bEnablePlugin;
 
 	// Allow FileServer connection using network
@@ -70,7 +70,7 @@ class UAndroidFileServerRuntimeSettings : public UObject
 	TEnumAsByte<EAFSConnectionType::Type> ConnectionType;
 
 	// Use manual IP address instead of automatic query from device (only for single device deploys!)
-	UPROPERTY(EditAnywhere, config, Category = Connection, Meta=(EditCondition = "ConnectionType == EConnectionType::NetworkOnly || ConnectionType == EConnectionType::Combined"), Meta=(DisplayName = "Use Manual IP Address?"))
+	UPROPERTY(EditAnywhere, config, Category = Connection, Meta=(EditCondition = "ConnectionType == EAFSConnectionType::NetworkOnly || ConnectionType == EAFSConnectionType::Combined"), Meta=(DisplayName = "Use Manual IP Address?"))
 	bool bUseManualIPAddress;
 
 	// IP address of device to use
