PR #9429: Fix broken Android build on macOS (Contributed by [REDACTED_PERSON])
Use rootProject path for AFSProject construction for AndroidFileServer

diff --git a/Engine/Plugins/Runtime/AndroidFileServer/Source/AndroidFileServer/AndroidFileServer_UPL.xml b/Engine/Plugins/Runtime/AndroidFileServer/Source/AndroidFileServer/AndroidFileServer_UPL.xml
--- a/Engine/Plugins/Runtime/AndroidFileServer/Source/AndroidFileServer/AndroidFileServer_UPL.xml
+++ b/Engine/Plugins/Runtime/AndroidFileServer/Source/AndroidFileServer/AndroidFileServer_UPL.xml
@@ -230,8 +230,8 @@ tasks.whenTaskAdded { task ->
 				<insert>
 <![CDATA[
 	// copy files into AFSStub
-	File afs_infile = new File('gradle.properties')
-	File afs_outfile = new File('AFSProject/gradle.properties')
+	File afs_infile = new File(rootProject.projectDir, 'gradle.properties')
+	File afs_outfile = new File(rootProject.projectDir, 'AFSProject/gradle.properties')
 	def afs_line = "";
 	afs_outfile.withWriter { writer ->
 		afs_infile.withReader { reader ->
@@ -249,8 +249,8 @@ tasks.whenTaskAdded { task ->
 		}
 	}
 
-	afs_infile = new File('app/buildAdditions.gradle')
-	afs_outfile = new File('AFSProject/app/buildAdditions.gradle')
+	afs_infile = new File(rootProject.projectDir, 'app/buildAdditions.gradle')
+	afs_outfile = new File(rootProject.projectDir, 'AFSProject/app/buildAdditions.gradle')
 	afs_line = "";
 	def bRemovingAFSBlock = false;
 	afs_outfile.withWriter { writer ->
@@ -274,8 +274,8 @@ tasks.whenTaskAdded { task ->
 		}
 	}
 
-	afs_infile = new File('app/src/main/AndroidManifest.xml')
-	afs_outfile = new File('AFSProject/app/src/main/AndroidManifest.xml')
+	afs_infile = new File(rootProject.projectDir, 'app/src/main/AndroidManifest.xml')
+	afs_outfile = new File(rootProject.projectDir, 'AFSProject/app/src/main/AndroidManifest.xml')
 	def afs_hasReceiver = false
 	afs_outfile.withWriter { writer ->
 		afs_infile.withReader { reader ->
@@ -300,9 +300,9 @@ tasks.whenTaskAdded { task ->
 		"baseBuildAdditions.gradle", "build.gradle", "buildscriptAdditions.gradle", "local.properties",
 		"gradlew", "gradlew.bat", "rungradle.bat"] as String[]
 	afs_files.eachWithIndex { v, i ->
-		afs_infile = new File(v)
+		afs_infile = new File(rootProject.projectDir, v)
 		if (afs_infile.exists()) {
-			afs_outfile = new File("AFSProject/" + v)
+			afs_outfile = new File(rootProject.projectDir, "AFSProject/" + v)
 			afs_outfile.withWriter { writer ->
 				afs_infile.withReader { reader ->
 					while ((afs_line = reader.readLine()) != null) {
