Fix issue with compiling of AFSProject on Mac and Linux

#android

diff --git a/Engine/Plugins/Runtime/AndroidFileServer/Source/AndroidFileServer/AndroidFileServer_UPL.xml b/Engine/Plugins/Runtime/AndroidFileServer/Source/AndroidFileServer/AndroidFileServer_UPL.xml
--- a/Engine/Plugins/Runtime/AndroidFileServer/Source/AndroidFileServer/AndroidFileServer_UPL.xml
+++ b/Engine/Plugins/Runtime/AndroidFileServer/Source/AndroidFileServer/AndroidFileServer_UPL.xml
@@ -186,7 +186,7 @@ tasks.whenTaskAdded { task ->
 			if (System.getProperty('os.name').toLowerCase(Locale.ROOT).contains('windows')) {
 				commandLine 'cmd', '/c', 'gradlew.bat', 'app:assembleRelease'
 			} else {
-				commandLine 'sh', '-c', 'chmod +x ./gradlew && ./gradlew', 'app:assembleRelease'
+				commandLine 'sh', '-c', 'chmod +x ./gradlew && ./gradlew app:assembleRelease'
 			}
 		}
 		task.finalizedBy 'ueAFSProjectAssembleRelease'
@@ -198,7 +198,7 @@ tasks.whenTaskAdded { task ->
 			if (System.getProperty('os.name').toLowerCase(Locale.ROOT).contains('windows')) {
 				commandLine 'cmd', '/c', 'gradlew.bat', 'app:assembleDebug'
 			} else {
-				commandLine 'sh', '-c', 'chmod +x ./gradlew && ./gradlew', 'app:assembleDebug'
+				commandLine 'sh', '-c', 'chmod +x ./gradlew && ./gradlew app:assembleDebug'
 			}
 		}
 		task.finalizedBy 'ueAFSProjectAssembleDebug'
@@ -210,7 +210,7 @@ tasks.whenTaskAdded { task ->
 			if (System.getProperty('os.name').toLowerCase(Locale.ROOT).contains('windows')) {
 				commandLine 'cmd', '/c', 'gradlew.bat', 'app:assembleRelease'
 			} else {
-				commandLine 'sh', '-c', 'chmod +x ./gradlew && ./gradlew', 'app:assembleRelease'
+				commandLine 'sh', '-c', 'chmod +x ./gradlew && ./gradlew app:assembleRelease'
 			}
 		}
 		task.finalizedBy 'ueAFSProjectBundleRelease'
@@ -222,7 +222,7 @@ tasks.whenTaskAdded { task ->
 			if (System.getProperty('os.name').toLowerCase(Locale.ROOT).contains('windows')) {
 				commandLine 'cmd', '/c', 'gradlew.bat', 'app:assembleDebug'
 			} else {
-				commandLine 'sh', '-c', 'chmod +x ./gradlew && ./gradlew', 'app:assembleDebug'
+				commandLine 'sh', '-c', 'chmod +x ./gradlew && ./gradlew app:assembleDebug'
 			}
 		}
 		task.finalizedBy 'ueAFSProjectBundleDebug'
