Updates to handling of arbitrary OBB overflow files

diff --git a/Engine/Plugins/Runtime/AndroidFileServer/Source/AndroidFileServer/Android/java/com/epicgames/unreal/RemoteFileManager.java b/Engine/Plugins/Runtime/AndroidFileServer/Source/AndroidFileServer/Android/java/com/epicgames/unreal/RemoteFileManager.java
--- a/Engine/Plugins/Runtime/AndroidFileServer/Source/AndroidFileServer/Android/java/com/epicgames/unreal/RemoteFileManager.java
+++ b/Engine/Plugins/Runtime/AndroidFileServer/Source/AndroidFileServer/Android/java/com/epicgames/unreal/RemoteFileManager.java
@@ -56,6 +56,8 @@ public class RemoteFileManager implements Runnable
 	private static final int SERVER_VERSION = 100;
 	private static final int ACCEPT_TIMEOUT = 2000;		// milliseconds
 
+	private static final int MAX_OVERFLOW_FILES = 32;
+
 	private static final int Command_Terminate = 0;
 	private static final int Command_Close = 1;
 	private static final int Command_Info = 2;
@@ -93,10 +95,6 @@ public class RemoteFileManager implements Runnable
 	private String commandfile;
 	private String mainobb;
 	private String patchobb;
-	private String overflow1obb;
-	private String overflow2obb;
-	private String overflow3obb;
-	private String overflow4obb;
 
 	private String packageName;
 	private int versionCode;
@@ -147,10 +145,6 @@ public class RemoteFileManager implements Runnable
 		obbPackageDir = obbdir + "/";
 		mainobb = obbPackageDir + "main." + versionString + "." + packageName + ".obb";
 		patchobb = obbPackageDir + "patch." + versionString + "." + packageName + ".obb";
-		overflow1obb = obbPackageDir + "overflow1." + versionString + "." + packageName + ".obb";
-		overflow2obb = obbPackageDir + "overflow2." + versionString + "." + packageName + ".obb";
-		overflow3obb = obbPackageDir + "overflow3." + versionString + "." + packageName + ".obb";
-		overflow4obb = obbPackageDir + "overflow4." + versionString + "." + packageName + ".obb";
 
 		unrealPackageDir = (bIsShipping ? internalPackageDir : externalPackageDir) + unrealGame;
 		projectPackageDir = unrealPackageDir + projectName + "/";
@@ -173,10 +167,12 @@ public class RemoteFileManager implements Runnable
 		replacements.put("^obb", obbPackageDir);
 		replacements.put("^mainobb", mainobb);
 		replacements.put("^patchobb", patchobb);
-		replacements.put("^overflow1obb", overflow1obb);
-		replacements.put("^overflow2obb", overflow2obb);
-		replacements.put("^overflow3obb", overflow3obb);
-		replacements.put("^overflow4obb", overflow4obb);
+
+		for (int overflowIndex=1; overflowIndex < MAX_OVERFLOW_FILES; overflowIndex++)
+		{
+			String overflowobb = obbPackageDir + "overflow" + overflowIndex + "." + versionString + "." + packageName + ".obb";
+			replacements.put("^overflow" + overflowIndex + "obb", overflowobb);
+		}
 
 		replacements.put("^engineversion", engineVersion);
 		replacements.put("^projectname", projectName);
diff --git a/Engine/Source/Runtime/Core/Private/Android/AndroidPlatformFile.cpp b/Engine/Source/Runtime/Core/Private/Android/AndroidPlatformFile.cpp
--- a/Engine/Source/Runtime/Core/Private/Android/AndroidPlatformFile.cpp
+++ b/Engine/Source/Runtime/Core/Private/Android/AndroidPlatformFile.cpp
@@ -86,6 +86,8 @@ int32 GAndroidPackageVersion = 0;
 int32 GAndroidPackagePatchVersion = 0;
 FString GAndroidAppType;
 
+#define ANDROID_MAX_OVERFLOW_FILES	32
+
 // External File Path base - setup during load
 FString GFilePathBase;
 // Obb File Path base - setup during load
@@ -94,6 +96,10 @@ FString GOBBFilePathBase;
 FString GOBBMainFilePath;
 // Obb Patch filepath
 FString GOBBPatchFilePath;
+// Obb Overflow1 filepath
+FString GOBBOverflow1FilePath;
+// Obb Overflow2 filepath
+FString GOBBOverflow2FilePath;
 // Internal File Direcory Path (for application) - setup during load
 FString GInternalFilePath;
 // External File Direcory Path (for application) - setup during load
@@ -1242,8 +1248,23 @@ public:
 
 			// Only check for overflow files if we found a patch file
 			if (bHavePatch)
-			{	
-				for (int32 OverflowIndex = 1; ; OverflowIndex++)
+			{
+				int32 OverflowIndex = 1;
+
+				if (!GOBBOverflow1FilePath.IsEmpty() && FileExists(*GOBBOverflow1FilePath, true))
+				{
+					OverflowIndex = 2;
+					MountOBB(*GOBBOverflow1FilePath);
+					FPlatformMisc::LowLevelOutputDebugStringf(TEXT("Mounted overflow1 OBB: %s"), *GOBBOverflow1FilePath);
+				}
+				if (!GOBBOverflow2FilePath.IsEmpty() && FileExists(*GOBBOverflow2FilePath, true))
+				{
+					OverflowIndex = 3;
+					MountOBB(*GOBBOverflow2FilePath);
+					FPlatformMisc::LowLevelOutputDebugStringf(TEXT("Mounted overflow2 OBB: %s"), *GOBBOverflow2FilePath);
+				}
+
+				while (OverflowIndex <= ANDROID_MAX_OVERFLOW_FILES)
 				{
 					FString OverflowOBBName = FString::Printf(TEXT("overflow%d.%d.%s.obb"), OverflowIndex, GAndroidPackageVersion, *GPackageName);
 
@@ -1263,6 +1284,8 @@ public:
 					{
 						break;
 					}
+
+					OverflowIndex++;
 				}
 			}
 		}
diff --git a/Engine/Source/Runtime/Launch/Private/Android/AndroidJNI.cpp b/Engine/Source/Runtime/Launch/Private/Android/AndroidJNI.cpp
--- a/Engine/Source/Runtime/Launch/Private/Android/AndroidJNI.cpp
+++ b/Engine/Source/Runtime/Launch/Private/Android/AndroidJNI.cpp
@@ -45,6 +45,10 @@ extern FString GFontPathBase;
 extern bool GOBBinAPK;
 extern bool GOverrideAndroidLogDir;
 extern FString GOBBFilePathBase;
+extern FString GOBBMainFilePath;
+extern FString GOBBPatchFilePath;
+extern FString GOBBOverflow1FilePath;
+extern FString GOBBOverflow2FilePath;
 extern FString GAPKFilename;
 
 FOnActivityResult FJavaWrapper::OnActivityResultDelegate;
@@ -1667,6 +1671,10 @@ JNIEXPORT jint JNI_OnLoad(JavaVM* InJavaVM, void* InReserved)
 	// Copy that somewhere safe
 	GFilePathBase = FJavaHelper::FStringFromLocalRef(Env, (jstring)Env->CallObjectMethod(*externalStoragePath, getFilePath, nullptr));
 	GOBBFilePathBase = GFilePathBase;
+	GOBBMainFilePath = TEXT("");
+	GOBBPatchFilePath = TEXT("");
+	GOBBOverflow1FilePath = TEXT("");
+	GOBBOverflow2FilePath = TEXT("");
 
 	// then release...
 	FPlatformMisc::LowLevelOutputDebugStringf(TEXT("Path found as '%s'\n"), *GFilePathBase);
@@ -1690,7 +1698,10 @@ JNIEXPORT jint JNI_OnLoad(JavaVM* InJavaVM, void* InReserved)
 //This function is declared in the Java-defined class, GameActivity.java: "public native void naativeSetObbFilePaths();"
 JNI_METHOD void Java_com_epicgames_unreal_GameActivity_nativeSetObbFilePaths(JNIEnv* jenv, jobject thiz, jstring OBBMainFilePath, jstring OBBPatchFilePath, jstring OBBOverflow1FilePath, jstring OBBOverflow2FilePath)
 {
-	// no longer supported
+	GOBBMainFilePath = FJavaHelper::FStringFromParam(jenv, OBBMainFilePath);
+	GOBBPatchFilePath = FJavaHelper::FStringFromParam(jenv, OBBPatchFilePath);
+	GOBBOverflow1FilePath = FJavaHelper::FStringFromParam(jenv, OBBOverflow1FilePath);
+	GOBBOverflow2FilePath = FJavaHelper::FStringFromParam(jenv, OBBOverflow2FilePath);
 }
 
 //This function is declared in the Java-defined class, GameActivity.java: "public native void nativeSetGlobalActivity();"
