Convert AndroidFileServer to a foreground service

#android

diff --git a/Engine/Plugins/Runtime/AndroidFileServer/Source/AndroidFileServer/Android/java/com/epicgames/unreal/RemoteFileManagerActivity.java b/Engine/Plugins/Runtime/AndroidFileServer/Source/AndroidFileServer/Android/java/com/epicgames/unreal/RemoteFileManagerActivity.java
--- a/Engine/Plugins/Runtime/AndroidFileServer/Source/AndroidFileServer/Android/java/com/epicgames/unreal/RemoteFileManagerActivity.java
+++ b/Engine/Plugins/Runtime/AndroidFileServer/Source/AndroidFileServer/Android/java/com/epicgames/unreal/RemoteFileManagerActivity.java
@@ -3,40 +3,66 @@
 package com.epicgames.unreal;
 
 import android.app.Activity;
+import android.app.Notification;
+import android.app.NotificationManager;
+import android.app.NotificationChannel;
+import android.app.Service;
+import android.content.Context;
 import android.content.Intent;
-import android.content.pm.ApplicationInfo;
-import android.content.pm.PackageManager;
 import android.os.Bundle;
 import android.util.Log;
 
-import com.epicgames.unreal.GameActivity;
-import com.epicgames.unreal.RemoteFileManager;
-
-import java.io.BufferedReader;
-import java.io.IOException;
-import java.io.InputStreamReader;
-import java.io.PrintWriter;
-import java.net.ServerSocket;
-import java.net.Socket;
 import java.util.Set;
 
-import static android.os.Environment.getExternalStorageDirectory;
+import com.epicgames.unreal.RemoteFileManagerService;
+
 
 public class RemoteFileManagerActivity extends Activity {
     private static String TAG = "UEFS";
     private static final int SERVER_PORT = 57099;
 
+	private static final String NOTIFICATION_CHANNEL_ID = "unreal-afs-notification-channel-id";
+	private static final CharSequence NOTICATION_CHANNEL_NAME = "unreal-afs-notification-channel";
+
 	@Override
 	protected void onCreate(Bundle savedInstanceState) {
 		super.onCreate(savedInstanceState);
 
 		Intent intent = getIntent();
 		String action = intent.getAction();
-		if (action != null && action == "com.epicgames.unreal.RemoteFileManager.intent.COMMAND2") {
+		if (action != null && action.equals("com.epicgames.unreal.RemoteFileManager.intent.COMMAND2")) {
 			String cmd = intent.getStringExtra("cmd");
 			String token = intent.getStringExtra("token");
 			int port = intent.getIntExtra("port", SERVER_PORT);
-			if (cmd != null && cmd.equals("start") && GameActivity.AndroidFileServer_Verify(token)) {
+
+			Log.d(TAG, "RemoteFileManagerActivity cmd: " + cmd + ", package = " + getPackageName());
+
+			if (cmd.equals("stop"))
+			{
+				stopAllThreads();
+				stopService(new Intent(this, RemoteFileManagerService.class));
+			} else 	if (cmd != null && cmd.equals("start") && GameActivity.AndroidFileServer_Verify(token)) {
+
+				// make a service intent with same extras
+				Intent serviceIntent = new Intent(this, RemoteFileManagerService.class);
+				serviceIntent.setAction(action);
+				serviceIntent.putExtra("cmd", cmd);
+				serviceIntent.putExtra("token", token);
+				serviceIntent.putExtra("port", port);
+
+				// make sure the notification channel exists (create it if not)
+				NotificationManager notificationManager = (NotificationManager) getSystemService(NotificationManager.class);
+				NotificationChannel channel = notificationManager.getNotificationChannel(NOTIFICATION_CHANNEL_ID);
+				if (channel == null)
+				{
+					channel = new NotificationChannel(NOTIFICATION_CHANNEL_ID, NOTICATION_CHANNEL_NAME, NotificationManager.IMPORTANCE_LOW);
+					channel.enableVibration(false);
+					channel.enableLights(false);
+					channel.setSound(null, null);
+					channel.setShowBadge(false);
+					channel.setLockscreenVisibility(Notification.VISIBILITY_SECRET);
+					notificationManager.createNotificationChannel(channel);
+				}
 
 				// make sure UnrealFileServer not already running
 				boolean bUSBRunning = false;
@@ -61,91 +87,37 @@ public class RemoteFileManagerActivity extends Activity {
 					}
 				}
 
-				String packageName = getPackageName();
-				PackageManager pm = getPackageManager();
-
-				int versionCode = 0;
-				try
-				{
-					versionCode = pm.getPackageInfo(packageName, 0).versionCode;
-				}
-				catch (Exception e)
-				{
-					// if the above failed, then, we can't use obbs
-				}
-
-				String EngineVersion = "5.0.0";
-				String ProjectName = packageName;
-				ProjectName = ProjectName.substring(ProjectName.lastIndexOf('.') + 1);
-				boolean IsShipping = false;
-				boolean UseExternalFilesDir = false;
-				boolean PublicLogFiles = false;
-				try {
-					ApplicationInfo ai = pm.getApplicationInfo(packageName, PackageManager.GET_META_DATA);
-					Bundle bundle = ai.metaData;
-
-					if (bundle.containsKey("com.epicgames.unreal.GameActivity.EngineVersion"))
-					{
-						EngineVersion = bundle.getString("com.epicgames.unreal.GameActivity.EngineVersion");
-					}
-					if (bundle.containsKey("com.epicgames.unreal.GameActivity.ProjectName"))
-					{
-						ProjectName = bundle.getString("com.epicgames.unreal.GameActivity.ProjectName");
-					}
-					if (bundle.containsKey("com.epicgames.unreal.GameActivity.BuildConfiguration"))
-					{
-						String Configuration = bundle.getString("com.epicgames.unreal.GameActivity.BuildConfiguration");
-						IsShipping = Configuration.contains("Shipping");
-					}
-					if (bundle.containsKey("com.epicgames.unreal.GameActivity.bUseExternalFilesDir"))
-					{
-						UseExternalFilesDir = bundle.getBoolean("com.epicgames.unreal.GameActivity.bUseExternalFilesDir");
-					}
-					if (bundle.containsKey("com.epicgames.unreal.GameActivity.bPublicLogFiles"))
-					{
-						PublicLogFiles = bundle.getBoolean("com.epicgames.unreal.GameActivity.bPublicLogFiles");
-					}
-				}
-				catch (PackageManager.NameNotFoundException | NullPointerException e)
-				{
-					Log.e("UE","Error when accessing application metadata", e);
-				}
-
-				String internal = getFilesDir().getAbsolutePath();
-				String external = getExternalFilesDir(null).getAbsolutePath();
-				String storage = getExternalStorageDirectory().getAbsolutePath();
-				String obbdir = getObbDir().getAbsolutePath();
-
-				if (!bUSBRunning)
+				if (bUSBRunning && bWiFiRunning)
 				{
-					RemoteFileManager fileManager = new RemoteFileManager(true, port, internal, external, storage, obbdir, packageName, versionCode, ProjectName, EngineVersion, IsShipping, PublicLogFiles);
-					new Thread(fileManager, "UnrealAFS-USB").start();
+					Log.d(TAG, "Both server threads already active on service");
 				}
-				if (!bWiFiRunning)
+				else
 				{
-					RemoteFileManager fileManager = new RemoteFileManager(false, port, internal, external, storage, obbdir, packageName, versionCode, ProjectName, EngineVersion, IsShipping, PublicLogFiles);
-					new Thread(fileManager, "UnrealAFS-WiFi").start();
+					// start the service
+					getApplicationContext().startForegroundService(serviceIntent);
 				}
 			}
-			else if (cmd != null && cmd.equals("stop")) {
-
-				// send stop request to threads
-				Set<Thread> threads = Thread.getAllStackTraces().keySet();
-				for (Thread thread : threads) {
-					String name = thread.getName();
-					if (name.equals("UnrealAFS-USB") || name.equals("UnrealAFS-WiFi"))
-					{
-						if (thread.isAlive())
-						{
-							Log.d(TAG, "Sent stop request to " + getPackageName() + "/" + name);
-							thread.interrupt();
-						}
-					}
-				}
-            }
 		}
 
 		// now terminate
 		finish();
     }
-}
\ No newline at end of file
+
+	private void stopAllThreads()
+	{
+		String packageName = getPackageName();
+		// send stop request to threads
+		Set<Thread> threads = Thread.getAllStackTraces().keySet();
+		for (Thread thread : threads) {
+			String name = thread.getName();
+			if (name.equals("UnrealAFS-USB") || name.equals("UnrealAFS-WiFi"))
+			{
+				if (thread.isAlive())
+				{
+					Log.d(TAG, "Sent stop request to " + packageName + "/" + name);
+					thread.interrupt();
+				}
+			}
+		}
+	}
+}
diff --git a/Engine/Plugins/Runtime/AndroidFileServer/Source/AndroidFileServer/Android/java/com/epicgames/unreal/RemoteFileManagerService.java b/Engine/Plugins/Runtime/AndroidFileServer/Source/AndroidFileServer/Android/java/com/epicgames/unreal/RemoteFileManagerService.java
new file mode 100644
--- /dev/null
+++ b/Engine/Plugins/Runtime/AndroidFileServer/Source/AndroidFileServer/Android/java/com/epicgames/unreal/RemoteFileManagerService.java
@@ -0,0 +1,252 @@
+// Copyright Epic Games, Inc. All Rights Reserved.
+
+package com.epicgames.unreal;
+
+import android.app.Service;
+import android.content.Context;
+import android.content.Intent;
+import android.content.pm.ApplicationInfo;
+import android.content.pm.PackageManager;
+import android.os.Bundle;
+import android.os.IBinder;
+import android.util.Log;
+
+import android.app.Notification;
+import android.app.NotificationManager;
+import android.app.NotificationChannel;
+import android.app.PendingIntent;
+
+import com.epicgames.unreal.GameActivity;
+import com.epicgames.unreal.RemoteFileManager;
+
+import java.io.BufferedReader;
+import java.io.IOException;
+import java.io.InputStreamReader;
+import java.io.PrintWriter;
+import java.net.ServerSocket;
+import java.net.Socket;
+import java.util.Set;
+
+import static android.os.Environment.getExternalStorageDirectory;
+
+public class RemoteFileManagerService extends Service {
+    private static String TAG = "UEFS";
+    private static final int SERVER_PORT = 57099;
+
+	private static final String NOTIFICATION_CHANNEL_ID = "unreal-afs-notification-channel-id";
+	private static final CharSequence NOTICATION_CHANNEL_NAME = "unreal-afs-notification-channel";
+
+	@Override
+	public int onStartCommand(Intent intent, int flags, int startId)
+	{
+		NotificationManager notificationManager = (NotificationManager) getSystemService(NotificationManager.class);
+		NotificationChannel channel = notificationManager.getNotificationChannel(NOTIFICATION_CHANNEL_ID);
+
+		if (channel != null)
+		{
+			Intent notificationIntent = new Intent(this, RemoteFileManagerService.class);
+			PendingIntent pendingIntent = PendingIntent.getActivity(this, 0, notificationIntent, PendingIntent.FLAG_IMMUTABLE);
+
+			Notification.Builder notificationBuilder = new Notification.Builder(this, NOTIFICATION_CHANNEL_ID)
+				.setSmallIcon(getNotificationIconID(this))
+				.setContentTitle("Unreal Android File Server")
+				.setContentText("Running AFS")
+				.setContentIntent(pendingIntent)
+				.setTicker("UE-AFS")
+				.setSound(null)
+				.setOngoing(true);
+			if (android.os.Build.VERSION.SDK_INT >= 31)
+			{
+				notificationBuilder.setForegroundServiceBehavior(Notification.FOREGROUND_SERVICE_DEFERRED);
+			}
+
+			startForeground(1, notificationBuilder.build());
+		}
+		else
+		{
+			Log.d(TAG, "Notification channel not available; service not started");
+			stopForeground(Service.STOP_FOREGROUND_REMOVE);
+			stopSelf(startId);
+			return START_NOT_STICKY;
+		}
+
+		String action = intent.getAction();
+		if (action != null && action.equals("com.epicgames.unreal.RemoteFileManager.intent.COMMAND2")) {
+			String cmd = intent.getStringExtra("cmd");
+			String token = intent.getStringExtra("token");
+			int port = intent.getIntExtra("port", SERVER_PORT);
+			if (cmd != null && cmd.equals("start") && GameActivity.AndroidFileServer_Verify(token)) {
+				String packageName = getPackageName();
+				PackageManager pm = getPackageManager();
+
+				Log.d(TAG, "onStartCommand " + startId + " received start command " + packageName);
+
+				// make sure UnrealFileServer not already running
+				boolean bUSBRunning = false;
+				boolean bWiFiRunning = false;
+				
+				Set<Thread> threads = Thread.getAllStackTraces().keySet();
+				for (Thread thread : threads) {
+					String name = thread.getName();
+					if (name.equals("UnrealAFS-USB"))
+					{
+						if (thread.isAlive())
+						{
+							bUSBRunning = true;
+						}
+					}
+					else if (name.equals("UnrealAFS-WiFi"))
+					{
+						if (thread.isAlive())
+						{
+							bWiFiRunning = true;
+						}
+					}
+				}
+
+				if (bUSBRunning && bWiFiRunning)
+				{
+					Log.d(TAG, "Both server threads already active");
+					return START_REDELIVER_INTENT;
+				}
+
+				int versionCode = 0;
+				try
+				{
+					versionCode = pm.getPackageInfo(packageName, 0).versionCode;
+				}
+				catch (Exception e)
+				{
+					// if the above failed, then, we can't use obbs
+				}
+
+				String EngineVersion = "5.0.0";
+				String ProjectName = packageName;
+				ProjectName = ProjectName.substring(ProjectName.lastIndexOf('.') + 1);
+				boolean IsShipping = false;
+				boolean UseExternalFilesDir = false;
+				boolean PublicLogFiles = false;
+				try {
+					ApplicationInfo ai = pm.getApplicationInfo(packageName, PackageManager.GET_META_DATA);
+					Bundle bundle = ai.metaData;
+
+					if (bundle.containsKey("com.epicgames.unreal.GameActivity.EngineVersion"))
+					{
+						EngineVersion = bundle.getString("com.epicgames.unreal.GameActivity.EngineVersion");
+					}
+					if (bundle.containsKey("com.epicgames.unreal.GameActivity.ProjectName"))
+					{
+						ProjectName = bundle.getString("com.epicgames.unreal.GameActivity.ProjectName");
+					}
+					if (bundle.containsKey("com.epicgames.unreal.GameActivity.BuildConfiguration"))
+					{
+						String Configuration = bundle.getString("com.epicgames.unreal.GameActivity.BuildConfiguration");
+						IsShipping = Configuration.contains("Shipping");
+					}
+					if (bundle.containsKey("com.epicgames.unreal.GameActivity.bUseExternalFilesDir"))
+					{
+						UseExternalFilesDir = bundle.getBoolean("com.epicgames.unreal.GameActivity.bUseExternalFilesDir");
+					}
+					if (bundle.containsKey("com.epicgames.unreal.GameActivity.bPublicLogFiles"))
+					{
+						PublicLogFiles = bundle.getBoolean("com.epicgames.unreal.GameActivity.bPublicLogFiles");
+					}
+				}
+				catch (PackageManager.NameNotFoundException | NullPointerException e)
+				{
+					Log.e(TAG, "Error when accessing application metadata", e);
+				}
+
+				String internal = getFilesDir().getAbsolutePath();
+				String external = getExternalFilesDir(null).getAbsolutePath();
+				String storage = getExternalStorageDirectory().getAbsolutePath();
+				String obbdir = getObbDir().getAbsolutePath();
+
+				if (!bUSBRunning)
+				{
+					RemoteFileManager fileManager = new RemoteFileManager(true, port, internal, external, storage, obbdir, packageName, versionCode, ProjectName, EngineVersion, IsShipping, PublicLogFiles);
+					new Thread(fileManager, "UnrealAFS-USB").start();
+				}
+				if (!bWiFiRunning)
+				{
+					RemoteFileManager fileManager = new RemoteFileManager(false, port, internal, external, storage, obbdir, packageName, versionCode, ProjectName, EngineVersion, IsShipping, PublicLogFiles);
+					new Thread(fileManager, "UnrealAFS-WiFi").start();
+				}
+
+				return START_REDELIVER_INTENT;
+			}
+			else if (cmd != null && cmd.equals("start")) {
+				// failed token verify - ignore it
+				String packageName = getPackageName();
+				Log.d(TAG, "onStartCommand " + startId + " received start command: " + packageName + " BAD TOKEN");
+				return START_NOT_STICKY;
+			}
+			else if (cmd != null && cmd.equals("stop")) {
+				String packageName = getPackageName();
+				Log.d(TAG, "onStartCommand " + startId + " received stop command: " + packageName);
+				stopAllThreads();
+				stopForeground(Service.STOP_FOREGROUND_REMOVE);
+				stopSelf(startId);
+				return START_NOT_STICKY;
+			}
+		}
+
+		Log.d(TAG, "onStartCommand received unknown command - NOT STARTING");
+		return START_NOT_STICKY;
+	}
+	
+	public static int getNotificationIconID(Context context)
+	{
+		int notificationIconID = context.getResources().getIdentifier("ic_notification_simple", "drawable", context.getPackageName());
+		if (notificationIconID == 0)
+		{
+			notificationIconID = context.getResources().getIdentifier("ic_notification", "drawable", context.getPackageName());
+		}
+		if (notificationIconID == 0)
+		{
+			notificationIconID = context.getResources().getIdentifier("icon", "drawable", context.getPackageName());
+		}
+		return notificationIconID;
+	}
+
+	private void stopAllThreads()
+	{
+		String packageName = getPackageName();
+		// send stop request to threads
+		Set<Thread> threads = Thread.getAllStackTraces().keySet();
+		for (Thread thread : threads) {
+			String name = thread.getName();
+			if (name.equals("UnrealAFS-USB") || name.equals("UnrealAFS-WiFi"))
+			{
+				if (thread.isAlive())
+				{
+					Log.d(TAG, "Sent stop request to " + packageName + "/" + name);
+					thread.interrupt();
+				}
+			}
+		}
+	}
+
+	@Override
+	public IBinder onBind(Intent intent)
+	{
+		// binding not supported
+		return null;
+	}
+
+	@Override
+	public void onDestroy()
+	{
+		super.onDestroy();
+//		String packageName = getPackageName();
+//		Log.d(TAG, "onDestroy " + packageName);
+	}
+
+	@Override
+	public void onCreate()
+	{
+		super.onCreate();
+//		String packageName = getPackageName();
+//		Log.d(TAG, "onCreate " + packageName);
+    }
+}
diff --git a/Engine/Plugins/Runtime/AndroidFileServer/Source/AndroidFileServer/AndroidFileServer_UPL.xml b/Engine/Plugins/Runtime/AndroidFileServer/Source/AndroidFileServer/AndroidFileServer_UPL.xml
--- a/Engine/Plugins/Runtime/AndroidFileServer/Source/AndroidFileServer/AndroidFileServer_UPL.xml
+++ b/Engine/Plugins/Runtime/AndroidFileServer/Source/AndroidFileServer/AndroidFileServer_UPL.xml
@@ -62,7 +62,13 @@
 					<true>
 						<if condition="bExternalStartAFS">
 							<true>
+								<!-- runtime permissions for foreground service -->
+								<addPermission android:name="android.permission.FOREGROUND_SERVICE"/>
+								<addPermission android:name="android.permission.FOREGROUND_SERVICE_DATA_SYNC"/>
+								<addPermission android:name="android.permission.POST_NOTIFICATIONS"/>
+								
 								<addElements tag="application">
+									<service android:name="com.epicgames.unreal.RemoteFileManagerService" android:foregroundServiceType="dataSync" android:exported="true" />
 									<activity android:name="com.epicgames.unreal.RemoteFileManagerActivity" android:exported="true" android:screenOrientation="sensor" android:theme="@android:style/Theme.NoDisplay">
 										<intent-filter>
 											<action android:name="com.epicgames.unreal.RemoteFileManager.intent.COMMAND2" />
@@ -87,7 +93,7 @@
 						</if>
 					</true>
 				</if>
-				
+
 				<!-- force bExternalFilesDir to true for all builds -->
 				<setBool result="bUpdatedExternal" value="false"/>
 				<loopElements tag="meta-data">
@@ -126,13 +132,16 @@
 				<if condition="bExternalStartAFS">
 					<true>
 						<copyFile src="$S(PluginDir)/Android/java/com/epicgames/unreal/RemoteFileManagerActivity.java" dst="$S(BuildDir)/src/com/epicgames/unreal/RemoteFileManagerActivity.java" />
+						<copyFile src="$S(PluginDir)/Android/java/com/epicgames/unreal/RemoteFileManagerService.java" dst="$S(BuildDir)/src/com/epicgames/unreal/RemoteFileManagerService.java" />
 					</true>
 					<false>
 						<isArch arch="arm64-v8a">
 							<deleteFiles filespec="arm64/src/com/epicgames/unreal/RemoteFileManagerActivity.java" />
+							<deleteFiles filespec="arm64/src/com/epicgames/unreal/RemoteFileManagerService.java" />
 						</isArch>
 						<isArch arch="x86_64">
 							<deleteFiles filespec="x64/src/com/epicgames/unreal/RemoteFileManagerActivity.java" />
+							<deleteFiles filespec="x64/src/com/epicgames/unreal/RemoteFileManagerService.java" />
 						</isArch>
 					</false>
 				</if>
@@ -141,10 +150,12 @@
 				<isArch arch="arm64-v8a">
 					<deleteFiles filespec="arm64/src/com/epicgames/unreal/RemoteFileManager.java" />
 					<deleteFiles filespec="arm64/src/com/epicgames/unreal/RemoteFileManagerActivity.java" />
+					<deleteFiles filespec="arm64/src/com/epicgames/unreal/RemoteFileManagerService.java" />
 				</isArch>
 				<isArch arch="x86_64">
 					<deleteFiles filespec="x64/src/com/epicgames/unreal/RemoteFileManager.java" />
 					<deleteFiles filespec="x64/src/com/epicgames/unreal/RemoteFileManagerActivity.java" />
+					<deleteFiles filespec="x64/src/com/epicgames/unreal/RemoteFileManagerService.java" />
 				</isArch>
 			</false>
 		</if>
@@ -276,13 +287,28 @@ tasks.whenTaskAdded { task ->
 
 	afs_infile = new File(rootProject.projectDir, 'app/src/main/AndroidManifest.xml')
 	afs_outfile = new File(rootProject.projectDir, 'AFSProject/app/src/main/AndroidManifest.xml')
+	def afs_hasForeServiceDataPerm = false;
+	def afs_hasForeServicePerm = false;
+	def afs_hasNotificationsPerm = false;
+	def afs_hasService = false
 	def afs_hasActivity = false
 	afs_outfile.withWriter { writer ->
 		afs_infile.withReader { reader ->
 			while ((afs_line = reader.readLine()) != null) {
-				if (afs_line.contains('com.epicgames.unreal.RemoteFileManagerActivity')) {
+				if (afs_line.contains('android.permission.FOREGROUND_SERVICE_DATA_SYNC')) {
+					afs_hasForeServiceDataPerm = true
+				} else if (afs_line.contains('android.permission.FOREGROUND_SERVICE')) {
+					afs_hasForeServicePerm = true
+				} else if (afs_line.contains('android.permission.POST_NOTIFICATIONS')) {
+					afs_hasNotificationsPerm = true
+				} else if (afs_line.contains('com.epicgames.unreal.RemoteFileManagerService')) {
+					afs_hasService = true
+				} else if (afs_line.contains('com.epicgames.unreal.RemoteFileManagerActivity')) {
 					afs_hasActivity = true
 				} else if (afs_line.contains('</application>')) {
+					if (afs_hasService == false) {
+						writer.writeLine('    <service android:name="com.epicgames.unreal.RemoteFileManagerService" android:foregroundServiceType="dataSync" android:exported="true" />')
+					}
 					if (afs_hasActivity == false) {
 						writer.writeLine('    <activity android:name="com.epicgames.unreal.RemoteFileManagerActivity" android:exported="true" android:screenOrientation="sensor">')
 						writer.writeLine('    	<intent-filter>')
@@ -290,6 +316,16 @@ tasks.whenTaskAdded { task ->
 						writer.writeLine('    	</intent-filter>')
 						writer.writeLine('    </activity>')
 					}
+				} else if (afs_line.contains('</manifest>')) {
+					if (afs_hasForeServicePerm == false) {
+						writer.writeLine('  <uses-permission android:name="android.permission.FOREGROUND_SERVICE" />')
+					}
+					if (afs_hasForeServiceDataPerm == false) {
+						writer.writeLine('  <uses-permission android:name="android.permission.FOREGROUND_SERVICE_DATA_SYNC" />')
+					}
+					if (afs_hasNotificationsPerm == false) {
+						writer.writeLine('  <uses-permission android:name="android.permission.POST_NOTIFICATIONS" />')
+					}
 				}
 				writer.writeLine("${afs_line}")
 			}
