Remove need for secuity token for debuggable applications when using AndroidFileServer with UnrealAndroidFileTool

#android

diff --git a/Engine/Plugins/Runtime/AndroidFileServer/Source/AFSStub/app/src/main/java/com/epicgames/ue4/GameApplication.java b/Engine/Plugins/Runtime/AndroidFileServer/Source/AFSStub/app/src/main/java/com/epicgames/ue4/GameApplication.java
--- a/Engine/Plugins/Runtime/AndroidFileServer/Source/AFSStub/app/src/main/java/com/epicgames/ue4/GameApplication.java
+++ b/Engine/Plugins/Runtime/AndroidFileServer/Source/AFSStub/app/src/main/java/com/epicgames/ue4/GameApplication.java
@@ -14,6 +14,11 @@ public class GameApplication extends Application implements LifecycleObserver {
 
 	private static boolean isForeground = false;
 
+	public static boolean AndroidFileServer_Init(String filename)
+	{
+		return false;
+	}
+
 	public static boolean AndroidFileServer_Verify(String Token)
 	{
 		return true;
diff --git a/Engine/Plugins/Runtime/AndroidFileServer/Source/AFSStub/app/src/main/java/com/epicgames/unreal/GameApplication.java b/Engine/Plugins/Runtime/AndroidFileServer/Source/AFSStub/app/src/main/java/com/epicgames/unreal/GameApplication.java
--- a/Engine/Plugins/Runtime/AndroidFileServer/Source/AFSStub/app/src/main/java/com/epicgames/unreal/GameApplication.java
+++ b/Engine/Plugins/Runtime/AndroidFileServer/Source/AFSStub/app/src/main/java/com/epicgames/unreal/GameApplication.java
@@ -14,6 +14,11 @@ public class GameApplication extends Application implements LifecycleObserver {
 
 	private static boolean isForeground = false;
 
+	public static boolean AndroidFileServer_Init(String filename)
+	{
+		return false;
+	}
+
 	public static boolean AndroidFileServer_Verify(String Token)
 	{
 		return true;
diff --git a/Engine/Plugins/Runtime/AndroidFileServer/Source/AndroidFileServer/Android/java/com/epicgames/unreal/RemoteFileManagerActivity.java b/Engine/Plugins/Runtime/AndroidFileServer/Source/AndroidFileServer/Android/java/com/epicgames/unreal/RemoteFileManagerActivity.java
--- a/Engine/Plugins/Runtime/AndroidFileServer/Source/AndroidFileServer/Android/java/com/epicgames/unreal/RemoteFileManagerActivity.java
+++ b/Engine/Plugins/Runtime/AndroidFileServer/Source/AndroidFileServer/Android/java/com/epicgames/unreal/RemoteFileManagerActivity.java
@@ -37,11 +37,16 @@ public class RemoteFileManagerActivity extends Activity {
 
 			Log.d(TAG, "RemoteFileManagerActivity cmd: " + cmd + ", package = " + getPackageName());
 
-			if (cmd.equals("stop"))
+			if ("init".equals(cmd))
+			{
+				GameApplication.AndroidFileServer_Init(getFilesDir().getAbsolutePath() + "/afs.cache");
+			}
+			else if ("stop".equals(cmd))
 			{
 				stopAllThreads();
 				stopService(new Intent(this, RemoteFileManagerService.class));
-			} else 	if (cmd != null && cmd.equals("start") && GameApplication.AndroidFileServer_Verify(token)) {
+			}
+			else if ("start".equals(cmd) && GameApplication.AndroidFileServer_Verify(token)) {
 
 				// make a service intent with same extras
 				Intent serviceIntent = new Intent(this, RemoteFileManagerService.class);
diff --git a/Engine/Plugins/Runtime/AndroidFileServer/Source/AndroidFileServer/AndroidFileServer_UPL.xml b/Engine/Plugins/Runtime/AndroidFileServer/Source/AndroidFileServer/AndroidFileServer_UPL.xml
--- a/Engine/Plugins/Runtime/AndroidFileServer/Source/AndroidFileServer/AndroidFileServer_UPL.xml
+++ b/Engine/Plugins/Runtime/AndroidFileServer/Source/AndroidFileServer/AndroidFileServer_UPL.xml
@@ -357,13 +357,27 @@ tasks.whenTaskAdded { task ->
 		</if>
 	</settingsGradleAdditions>
 
+	<gameApplicationImportAdditions>
+		<if condition="bEmbedAFS">
+			<true>
+				<insert>
+import java.io.BufferedReader;
+import java.io.File;
+import java.io.FileReader;
+import java.io.FileWriter;
+import java.io.IOException;
+				</insert>
+			</true>
+		</if>
+	</gameApplicationImportAdditions>
+
 	<gameActivityImportAdditions>
 		<if condition="bEmbedAFS">
 			<true>
 				<insert>
 import java.util.Set;
 import com.epicgames.unreal.RemoteFileManager;
-</insert>
+				</insert>
 			</true>
 		</if>
 	</gameActivityImportAdditions>
@@ -371,16 +385,75 @@ import com.epicgames.unreal.RemoteFileManager;
 	<gameApplicationClassAdditions>
 		<if condition="bEmbedAFS">
 			<true>
+				<if condition="Distribution">
+					<true>
+						<insert>
+	public static boolean AndroidFileServer_Init(String filename)
+	{
+		return false;
+	}
+						</insert>
+					</true>
+					<false>
+						<insert>
+	public static boolean AndroidFileServer_Init(String filename)
+	{
+		File file = new File(filename);
+		try
+		{
+			boolean bNeedsWrite = true;
+			if (file.exists())
+			{
+				String Token = "";
+				BufferedReader reader = null;
+				try
+				{
+					reader = new BufferedReader(new FileReader(file));
+					Token = reader.readLine();
+					reader.close();
+						</insert>
+						<insertValue value="&#9;&#9;&#9;&#9;&#9;if (&quot;$S(SecurityToken)&quot;.equals(Token))"/>
+						<insertNewline/>
+						<insert>
+					{
+						bNeedsWrite = false;
+					}
+				}
+				catch (IOException ie)
+				{
+				}
+			}
+			if (bNeedsWrite)
+			{
+
+				FileWriter outputWriter = new FileWriter(file);
+						</insert>
+						<insertValue value="&#9;&#9;&#9;&#9;outputWriter.write(&quot;$S(SecurityToken)&quot;);"/>
+						<insertNewline/>
+						<insert>
+				outputWriter.close();
+			}
+		}
+		catch (IOException e)
+		{
+			return false;
+		}
+		return true;
+	}
+						</insert>
+					</false>
+				</if>
+				<insertNewline/>
 				<insert>
 	public static boolean AndroidFileServer_Verify(String Token)
 	{
 				</insert>
 				<if condition="bTokenOptional">
 					<true>
-						<insertValue value="return true;"/>
+						<insertValue value="&#9;&#9;return true;"/>
 					</true>
 					<false>
-						<insertValue value="return (Token.equals(&quot;$S(SecurityToken)&quot;));"/>
+						<insertValue value="&#9;&#9;return (&quot;$S(SecurityToken)&quot;.equals(Token));"/>
 					</false>
 				</if>
 				<insertNewline/>
@@ -389,10 +462,17 @@ import com.epicgames.unreal.RemoteFileManager;
 				</insert>
 			</true>
 			<false>
+				<insert>
+	public static boolean AndroidFileServer_Init(String filename)
+	{
+		return true;
+	}
+
 	public static boolean AndroidFileServer_Verify(String Token)
 	{
 		return false;
 	}
+				</insert>
 			</false>
 		</if>
 	</gameApplicationClassAdditions>
diff --git a/Engine/Source/Programs/AutomationTool/Android/AndroidFileClient.cs b/Engine/Source/Programs/AutomationTool/Android/AndroidFileClient.cs
--- a/Engine/Source/Programs/AutomationTool/Android/AndroidFileClient.cs
+++ b/Engine/Source/Programs/AutomationTool/Android/AndroidFileClient.cs
@@ -1892,14 +1892,39 @@ namespace AutomationTool
 				}
 			}
 
+			// try to retrieve token if not provided
+			if (Token == "")
+			{
+				// try asking AFS to write the file; we do this instead of just checking for file since an update could have changed the token
+				ActivityManagerSendAndWait(Device, PackageName, "-e cmd init");
+
+				// this will only work if debuggable - try to retrieve the file contents
+				String CatResult = adb.Shell(Device, "run-as " + PackageName+" cat files/afs.cache");
+				if (CatResult.Contains("not debuggable"))
+				{
+					Logger.LogInformation("Package is not debuggable; assuming token not required.");
+				}
+				else if (!CatResult.Contains("No such file or directory"))
+				{
+					Token = CatResult;
+				}
+				else
+				{
+					Logger.LogInformation("Did not retrieve a token; assuming not required.");
+				}
+			}
+
 			// retries up to 8 seconds to start and see listener ready
 			tries = 40;
 			bool bUSB;
 			bool bWifi;
 			string WifiAddress;
+			
+			// only include -e Token if non-empty
+			String StartRequestParams = "-e cmd start" + (Token == "" ? "" : " -e token " + Token) + " -ei port " + ServerPort;
 
 			// sent start request (won't do anything if already started)
-			ActivityManagerSendAndWait(Device, PackageName, "-e cmd start -e token " + Token + " -ei port " + ServerPort);
+			ActivityManagerSendAndWait(Device, PackageName, StartRequestParams);
 
 			// see if we can check listen status
 			if (GetListenStatus(Device, ServerPort, out bUSB, out bWifi, out WifiAddress))
@@ -1919,7 +1944,7 @@ namespace AutomationTool
 					Thread.Sleep(200);
 
 					// sent start request again (won't do anything if already started)
-					ActivityManagerSendAndWait(Device, PackageName, "-e cmd start -e token " + Token + " -ei port " + ServerPort);
+					ActivityManagerSendAndWait(Device, PackageName, StartRequestParams);
 
 					GetListenStatus(Device, ServerPort, out bUSB, out bWifi, out WifiAddress);
 				}
@@ -1940,7 +1965,7 @@ namespace AutomationTool
 				}
 
 				// sent start request again (won't do anything if already started)
-				ActivityManagerSendAndWait(Device, PackageName, "-e cmd start -e token " + Token + " -ei port " + ServerPort);
+				ActivityManagerSendAndWait(Device, PackageName, StartRequestParams);
 			}
 
 			Logger.LogInformation("Timed out on connection attempts");
diff --git a/Engine/Source/Programs/UnrealAndroidFileTool/Program.cs b/Engine/Source/Programs/UnrealAndroidFileTool/Program.cs
--- a/Engine/Source/Programs/UnrealAndroidFileTool/Program.cs
+++ b/Engine/Source/Programs/UnrealAndroidFileTool/Program.cs
@@ -18,7 +18,7 @@ namespace UnrealAndroidFileTool
 {
 	class Program
 	{
-		private static string AppVersion = "1.0.6";
+		private static string AppVersion = "1.0.7";
 		private static ILogger Logger = Log.Logger;
 
 		static void ShowHelp(string command)
