Filter out AFSProject task if running it in Gradle

#android

diff --git a/Engine/Plugins/Runtime/AndroidFileServer/Source/AndroidFileServer/AndroidFileServer_UPL.xml b/Engine/Plugins/Runtime/AndroidFileServer/Source/AndroidFileServer/AndroidFileServer_UPL.xml
--- a/Engine/Plugins/Runtime/AndroidFileServer/Source/AndroidFileServer/AndroidFileServer_UPL.xml
+++ b/Engine/Plugins/Runtime/AndroidFileServer/Source/AndroidFileServer/AndroidFileServer_UPL.xml
@@ -180,53 +180,55 @@
 <![CDATA[
 // --- Begin AFSProject ---
 tasks.whenTaskAdded { task ->
-	if (task.name == 'assembleRelease') {
-		tasks.create(name:'ueAFSProjectAssembleRelease', type: Exec) {
-			description "Compile AFSProject for Release"
-			workingDir "$rootDir" + "/AFSProject"
-			if (System.getProperty('os.name').toLowerCase(Locale.ROOT).contains('windows')) {
-				commandLine 'cmd', '/c', 'gradlew.bat', 'app:assembleRelease'
-			} else {
-				commandLine 'sh', '-c', 'chmod +x ./gradlew && ./gradlew app:assembleRelease'
+	if (!System.getProperty("user.dir").contains("AFSProject")) {
+		if (task.name == 'assembleRelease') {
+			tasks.create(name:'ueAFSProjectAssembleRelease', type: Exec) {
+				description "Compile AFSProject for Release"
+				workingDir "$rootDir" + "/AFSProject"
+				if (System.getProperty('os.name').toLowerCase(Locale.ROOT).contains('windows')) {
+					commandLine 'cmd', '/c', 'gradlew.bat', 'app:assembleRelease'
+				} else {
+					commandLine 'sh', '-c', 'chmod +x ./gradlew && ./gradlew app:assembleRelease'
+				}
 			}
+			task.finalizedBy 'ueAFSProjectAssembleRelease'
 		}
-		task.finalizedBy 'ueAFSProjectAssembleRelease'
-	}
-	if (task.name == 'assembleDebug') {
-		tasks.create(name:'ueAFSProjectAssembleDebug', type: Exec) {
-			description "Compile AFSProject for Debug"
-			workingDir "$rootDir" + "/AFSProject"
-			if (System.getProperty('os.name').toLowerCase(Locale.ROOT).contains('windows')) {
-				commandLine 'cmd', '/c', 'gradlew.bat', 'app:assembleDebug'
-			} else {
-				commandLine 'sh', '-c', 'chmod +x ./gradlew && ./gradlew app:assembleDebug'
+		if (task.name == 'assembleDebug') {
+			tasks.create(name:'ueAFSProjectAssembleDebug', type: Exec) {
+				description "Compile AFSProject for Debug"
+				workingDir "$rootDir" + "/AFSProject"
+				if (System.getProperty('os.name').toLowerCase(Locale.ROOT).contains('windows')) {
+					commandLine 'cmd', '/c', 'gradlew.bat', 'app:assembleDebug'
+				} else {
+					commandLine 'sh', '-c', 'chmod +x ./gradlew && ./gradlew app:assembleDebug'
+				}
 			}
+			task.finalizedBy 'ueAFSProjectAssembleDebug'
 		}
-		task.finalizedBy 'ueAFSProjectAssembleDebug'
-	}
-	if (task.name == 'bundleRelease') {
-		tasks.create(name:'ueAFSProjectBundleRelease', type: Exec) {
-			description "Compile AFSProject for Release"
-			workingDir "$rootDir" + "/AFSProject"
-			if (System.getProperty('os.name').toLowerCase(Locale.ROOT).contains('windows')) {
-				commandLine 'cmd', '/c', 'gradlew.bat', 'app:assembleRelease'
-			} else {
-				commandLine 'sh', '-c', 'chmod +x ./gradlew && ./gradlew app:assembleRelease'
+		if (task.name == 'bundleRelease') {
+			tasks.create(name:'ueAFSProjectBundleRelease', type: Exec) {
+				description "Compile AFSProject for Release"
+				workingDir "$rootDir" + "/AFSProject"
+				if (System.getProperty('os.name').toLowerCase(Locale.ROOT).contains('windows')) {
+					commandLine 'cmd', '/c', 'gradlew.bat', 'app:assembleRelease'
+				} else {
+					commandLine 'sh', '-c', 'chmod +x ./gradlew && ./gradlew app:assembleRelease'
+				}
 			}
+			task.finalizedBy 'ueAFSProjectBundleRelease'
 		}
-		task.finalizedBy 'ueAFSProjectBundleRelease'
-	}
-	if (task.name == 'bundleDebug') {
-		tasks.create(name:'ueAFSProjectBundleDebug', type: Exec) {
-			description "Compile AFSProject for Debug"
-			workingDir "$rootDir" + "/AFSProject"
-			if (System.getProperty('os.name').toLowerCase(Locale.ROOT).contains('windows')) {
-				commandLine 'cmd', '/c', 'gradlew.bat', 'app:assembleDebug'
-			} else {
-				commandLine 'sh', '-c', 'chmod +x ./gradlew && ./gradlew app:assembleDebug'
+		if (task.name == 'bundleDebug') {
+			tasks.create(name:'ueAFSProjectBundleDebug', type: Exec) {
+				description "Compile AFSProject for Debug"
+				workingDir "$rootDir" + "/AFSProject"
+				if (System.getProperty('os.name').toLowerCase(Locale.ROOT).contains('windows')) {
+					commandLine 'cmd', '/c', 'gradlew.bat', 'app:assembleDebug'
+				} else {
+					commandLine 'sh', '-c', 'chmod +x ./gradlew && ./gradlew app:assembleDebug'
+				}
 			}
+			task.finalizedBy 'ueAFSProjectBundleDebug'
 		}
-		task.finalizedBy 'ueAFSProjectBundleDebug'
 	}
 }
 // --- End AFSProject ---
