Fix ASISDev build error

In CL 45443369, the return type for a method was accidentally removed in an apparent attempt to just remove static.
Replace the missing return type to fix the build.

Copied from CL 45541674 in Main.

diff --git a/Engine/Plugins/Runtime/AndroidFileServer/Source/AndroidFileServer/AndroidFileServer_UPL.xml b/Engine/Plugins/Runtime/AndroidFileServer/Source/AndroidFileServer/AndroidFileServer_UPL.xml
--- a/Engine/Plugins/Runtime/AndroidFileServer/Source/AndroidFileServer/AndroidFileServer_UPL.xml
+++ b/Engine/Plugins/Runtime/AndroidFileServer/Source/AndroidFileServer/AndroidFileServer_UPL.xml
@@ -474,7 +474,7 @@ import com.epicgames.unreal.RemoteFileManager;
 			</true>
 			<false>
 				<insert>
-	public AndroidFileServer_Init(String filename)
+	public boolean AndroidFileServer_Init(String filename)
 	{
 		return true;
 	}
