AFS interface included only on embed
#rnx

diff --git a/Engine/Plugins/Runtime/AndroidFileServer/Source/AndroidFileServer/AndroidFileServer_UPL.xml b/Engine/Plugins/Runtime/AndroidFileServer/Source/AndroidFileServer/AndroidFileServer_UPL.xml
--- a/Engine/Plugins/Runtime/AndroidFileServer/Source/AndroidFileServer/AndroidFileServer_UPL.xml
+++ b/Engine/Plugins/Runtime/AndroidFileServer/Source/AndroidFileServer/AndroidFileServer_UPL.xml
@@ -129,22 +129,20 @@
 		<if condition="bEmbedAFS">
 			<true>
 				<copyFile src="$S(PluginDir)/Android/java/com/epicgames/unreal/RemoteFileManager.java" dst="$S(BuildDir)/src/com/epicgames/unreal/RemoteFileManager.java" />
+				<copyFile src="$S(PluginDir)/Android/java/com/epicgames/unreal/IAndroidFileServerSetup.java" dst="$S(BuildDir)/src/com/epicgames/unreal/IAndroidFileServerSetup.java" />
 				<if condition="bExternalStartAFS">
 					<true>
 						<copyFile src="$S(PluginDir)/Android/java/com/epicgames/unreal/RemoteFileManagerActivity.java" dst="$S(BuildDir)/src/com/epicgames/unreal/RemoteFileManagerActivity.java" />
 						<copyFile src="$S(PluginDir)/Android/java/com/epicgames/unreal/RemoteFileManagerService.java" dst="$S(BuildDir)/src/com/epicgames/unreal/RemoteFileManagerService.java" />
-						<copyFile src="$S(PluginDir)/Android/java/com/epicgames/unreal/IAndroidFileServerSetup.java" dst="$S(BuildDir)/src/com/epicgames/unreal/IAndroidFileServerSetup.java" />
 					</true>
 					<false>
 						<isArch arch="arm64-v8a">
 							<deleteFiles filespec="arm64/src/com/epicgames/unreal/RemoteFileManagerActivity.java" />
 							<deleteFiles filespec="arm64/src/com/epicgames/unreal/RemoteFileManagerService.java" />
-							<deleteFiles filespec="arm64/src/com/epicgames/unreal/IAndroidFileServerSetup.java" />
 						</isArch>
 						<isArch arch="x86_64">
 							<deleteFiles filespec="x64/src/com/epicgames/unreal/RemoteFileManagerActivity.java" />
 							<deleteFiles filespec="x64/src/com/epicgames/unreal/RemoteFileManagerService.java" />
-							<deleteFiles filespec="x64/src/com/epicgames/unreal/IAndroidFileServerSetup.java" />
 						</isArch>
 					</false>
 				</if>
@@ -365,7 +363,11 @@ tasks.whenTaskAdded { task ->
 	</settingsGradleAdditions>
 
 	<gameApplicationImplementsAdditions>
-		<insert>, IAndroidFileServerSetup </insert>
+		<if condition="bEmbedAFS">
+			<true>
+				<insert>, IAndroidFileServerSetup </insert>
+			</true>
+		</if>
 	</gameApplicationImplementsAdditions>
 
 	<gameApplicationImportAdditions>
