Fix java compilation due to missing import

#rnx

diff --git a/Engine/Plugins/Runtime/AndroidFileServer/Source/AndroidFileServer/AndroidFileServer_UPL.xml b/Engine/Plugins/Runtime/AndroidFileServer/Source/AndroidFileServer/AndroidFileServer_UPL.xml
--- a/Engine/Plugins/Runtime/AndroidFileServer/Source/AndroidFileServer/AndroidFileServer_UPL.xml
+++ b/Engine/Plugins/Runtime/AndroidFileServer/Source/AndroidFileServer/AndroidFileServer_UPL.xml
@@ -379,6 +379,7 @@ import java.io.File;
 import java.io.FileReader;
 import java.io.FileWriter;
 import java.io.IOException;
+import com.epicgames.unreal.IAndroidFileServerSetup;
 				</insert>
 			</true>
 		</if>
