Unshelved from pending changelist '45879216':

Fix issue with AFSProject token validation

#android

diff --git a/Engine/Plugins/Runtime/AndroidFileServer/Source/AFSStub/app/src/main/java/com/epicgames/ue4/GameApplication.java b/Engine/Plugins/Runtime/AndroidFileServer/Source/AFSStub/app/src/main/java/com/epicgames/ue4/GameApplication.java
--- a/Engine/Plugins/Runtime/AndroidFileServer/Source/AFSStub/app/src/main/java/com/epicgames/ue4/GameApplication.java
+++ b/Engine/Plugins/Runtime/AndroidFileServer/Source/AFSStub/app/src/main/java/com/epicgames/ue4/GameApplication.java
@@ -7,19 +7,20 @@ import androidx.lifecycle.Lifecycle;
 import androidx.lifecycle.LifecycleObserver;
 import androidx.lifecycle.OnLifecycleEvent;
 import androidx.lifecycle.ProcessLifecycleOwner;
+import com.epicgames.unreal.IAndroidFileServerSetup;
 
 
-public class GameApplication extends Application implements LifecycleObserver {
+public class GameApplication extends Application implements LifecycleObserver , IAndroidFileServerSetup {
 	private static final Logger Log = new Logger("UE4", "GameApp");
 
 	private static boolean isForeground = false;
 
-	public static boolean AndroidFileServer_Init(String filename)
+	public boolean AndroidFileServer_Init(String filename)
 	{
 		return false;
 	}
 
-	public static boolean AndroidFileServer_Verify(String Token)
+	public boolean AndroidFileServer_Verify(String Token)
 	{
 		return true;
 	}
diff --git a/Engine/Plugins/Runtime/AndroidFileServer/Source/AFSStub/app/src/main/java/com/epicgames/unreal/GameApplication.java b/Engine/Plugins/Runtime/AndroidFileServer/Source/AFSStub/app/src/main/java/com/epicgames/unreal/GameApplication.java
--- a/Engine/Plugins/Runtime/AndroidFileServer/Source/AFSStub/app/src/main/java/com/epicgames/unreal/GameApplication.java
+++ b/Engine/Plugins/Runtime/AndroidFileServer/Source/AFSStub/app/src/main/java/com/epicgames/unreal/GameApplication.java
@@ -7,19 +7,20 @@ import androidx.lifecycle.Lifecycle;
 import androidx.lifecycle.LifecycleObserver;
 import androidx.lifecycle.OnLifecycleEvent;
 import androidx.lifecycle.ProcessLifecycleOwner;
+import com.epicgames.unreal.IAndroidFileServerSetup;
 
 
-public class GameApplication extends Application implements LifecycleObserver {
+public class GameApplication extends Application implements LifecycleObserver , IAndroidFileServerSetup {
 	private static final Logger Log = new Logger("UE", "GameApp");
 
 	private static boolean isForeground = false;
 
-	public static boolean AndroidFileServer_Init(String filename)
+	public boolean AndroidFileServer_Init(String filename)
 	{
 		return false;
 	}
 
-	public static boolean AndroidFileServer_Verify(String Token)
+	public boolean AndroidFileServer_Verify(String Token)
 	{
 		return true;
 	}
