Fix for timeout not shutting down AFS service properly
- also cleaned up some warnings

#Synced-CL -1

diff --git a/Engine/Plugins/Runtime/AndroidFileServer/Source/AndroidFileServer/Android/java/com/epicgames/unreal/RemoteFileManagerService.java b/Engine/Plugins/Runtime/AndroidFileServer/Source/AndroidFileServer/Android/java/com/epicgames/unreal/RemoteFileManagerService.java
--- a/Engine/Plugins/Runtime/AndroidFileServer/Source/AndroidFileServer/Android/java/com/epicgames/unreal/RemoteFileManagerService.java
+++ b/Engine/Plugins/Runtime/AndroidFileServer/Source/AndroidFileServer/Android/java/com/epicgames/unreal/RemoteFileManagerService.java
@@ -23,6 +23,7 @@ import java.io.BufferedReader;
 import java.io.IOException;
 import java.io.InputStreamReader;
 import java.io.PrintWriter;
+import java.io.File;
 import java.net.ServerSocket;
 import java.net.Socket;
 import java.util.Set;
@@ -53,7 +54,6 @@ public class RemoteFileManagerService extends Service {
 				.setContentText("Running AFS")
 				.setContentIntent(pendingIntent)
 				.setTicker("UE-AFS")
-				.setSound(null)
 				.setOngoing(true);
 			if (android.os.Build.VERSION.SDK_INT >= 31)
 			{
@@ -115,7 +115,14 @@ public class RemoteFileManagerService extends Service {
 				int versionCode = 0;
 				try
 				{
-					versionCode = pm.getPackageInfo(packageName, 0).versionCode;
+					if (android.os.Build.VERSION.SDK_INT >= 28)
+					{
+						versionCode = (int) pm.getPackageInfo(packageName, 0).getLongVersionCode();
+					}
+					else
+					{
+						versionCode = pm.getPackageInfo(packageName, 0).versionCode;
+					}
 				}
 				catch (Exception e)
 				{
@@ -161,18 +168,46 @@ public class RemoteFileManagerService extends Service {
 
 				String internal = getFilesDir().getAbsolutePath();
 				String external = getExternalFilesDir(null).getAbsolutePath();
-				String storage = getExternalStorageDirectory().getAbsolutePath();
+				File ExtStorageDir;
+				if (android.os.Build.VERSION.SDK_INT >= 29)
+				{
+					// Derive external storage root from getExternalFilesDir path (no deprecated API needed).
+					// external = .../Android/data/<pkg>/files; strip that suffix to get storage root.
+					String AndroidSuffix = "/Android/data/" + packageName + "/files";
+					if (external != null && external.endsWith(AndroidSuffix))
+					{
+						ExtStorageDir = new File(external.substring(0, external.length() - AndroidSuffix.length()));
+					}
+					else
+					{
+						Log.e(TAG, "Could not determine external storage root - using external files dir as fallback");
+						ExtStorageDir = null;
+					}
+				}
+				else
+				{
+					ExtStorageDir = android.os.Environment.getExternalStorageDirectory();
+					if (ExtStorageDir == null)
+					{
+						Log.e(TAG, "getExternalStorageDirectory() returned null - using external files dir as fallback");
+					}
+				}
+				String storage = (ExtStorageDir != null) ? ExtStorageDir.getAbsolutePath() : external;
 				String obbdir = getObbDir().getAbsolutePath();
 
 				if (!bUSBRunning)
 				{
 					RemoteFileManager fileManager = new RemoteFileManager(true, port, internal, external, storage, obbdir, packageName, versionCode, ProjectName, EngineVersion, IsShipping, PublicLogFiles);
-					new Thread(fileManager, "UnrealAFS-USB").start();
+					Thread usbThread = new Thread(fileManager, "UnrealAFS-USB");
+					usbThread.setUncaughtExceptionHandler((thread, ex) -> Log.e(TAG, "Server thread " + thread.getName() + " crashed: " + ex.toString()));
+					usbThread.start();
 				}
 				if (!bWiFiRunning)
 				{
 					RemoteFileManager fileManager = new RemoteFileManager(false, port, internal, external, storage, obbdir, packageName, versionCode, ProjectName, EngineVersion, IsShipping, PublicLogFiles);
-					new Thread(fileManager, "UnrealAFS-WiFi").start();
+					Thread wifiThread = new Thread(fileManager, "UnrealAFS-WiFi");
+					wifiThread.setUncaughtExceptionHandler((thread, ex) -> Log.e(TAG, "Server thread " + thread.getName() + " crashed: " + ex.toString()));
+					wifiThread.start();
 				}
 
 				return START_REDELIVER_INTENT;
@@ -236,6 +271,17 @@ public class RemoteFileManagerService extends Service {
 		return null;
 	}
 
+	// Android 14 has a 6-hour runtime limit.  onTimeOut() is called and we stop the service rather
+	// than having the system throw exception that shows up in crashes
+	@Override
+	public void onTimeout(int startId, int fgsType)
+	{
+		Log.w(TAG, "onTimeout: dataSync timeout reached, stopping service");
+		stopAllThreads();
+		stopForeground(Service.STOP_FOREGROUND_REMOVE);
+		stopSelf(startId);
+	}
+
 	@Override
 	public void onDestroy()
 	{
